Moved the AnalyticsMulticast module to be a plugin instead for Blueprint only project access

diff --git a/Engine/Plugins/Runtime/Analytics/AnalyticsMulticast/AnalyticsMulticast.uplugin b/Engine/Plugins/Runtime/Analytics/AnalyticsMulticast/AnalyticsMulticast.uplugin
new file mode 100644
--- /dev/null
+++ b/Engine/Plugins/Runtime/Analytics/AnalyticsMulticast/AnalyticsMulticast.uplugin
@@ -0,0 +1,21 @@
+{
+    "FileVersion": 3,
+
+    "FriendlyName": "Multicast Analytics Provider",
+    "Version": 1,
+    "VersionName": "1.0",
+    "CreatedBy": "Epic Games, Inc.",
+    "CreatedByURL": "http://epicgames.com",
+    "DocsURL": "https://docs.unrealengine.com/latest/INT/Gameplay/Analytics/index.html",
+    "EngineVersion": "4.8.0",
+    "Description": "Forwards analytics API calls to a list of analytics providers to log data to multiple services at once",
+    "Category": "Analytics.Provider",
+
+    "Modules": [
+        {
+            "Name": "AnalyticsMulticast",
+            "Type": "Runtime",
+            "LoadingPhase": "PreDefault"
+        }
+    ]
+}
diff --git a/Engine/Plugins/Runtime/Analytics/AnalyticsMulticast/Source/AnalyticsMulticast/AnalyticsMulticast.Build.cs b/Engine/Plugins/Runtime/Analytics/AnalyticsMulticast/Source/AnalyticsMulticast/AnalyticsMulticast.Build.cs
new file mode 100644
--- /dev/null
+++ b/Engine/Plugins/Runtime/Analytics/AnalyticsMulticast/Source/AnalyticsMulticast/AnalyticsMulticast.Build.cs
@@ -0,0 +1,30 @@
+// Copyright 1998-2015 Epic Games, Inc. All Rights Reserved.
+
+namespace UnrealBuildTool.Rules
+{
+	public class AnalyticsMulticast : ModuleRules
+	{
+		public AnalyticsMulticast(TargetInfo Target)
+		{
+            PublicDependencyModuleNames.AddRange(
+                new string[]
+				{
+					"Core",
+					"CoreUObject",
+                    "Engine",
+					// ... add other public dependencies that you statically link with here ...
+				}
+                );
+
+            PrivateDependencyModuleNames.AddRange(
+                new string[]
+				{
+					"Analytics",
+					// ... add private dependencies that you statically link with here ...
+				}
+                );
+
+            PublicIncludePathModuleNames.Add("Analytics");
+        }
+	}
+}
\ No newline at end of file
diff --git a/Engine/Plugins/Runtime/Analytics/AnalyticsMulticast/Source/AnalyticsMulticast/Private/AnalyticsMulticast.cpp b/Engine/Plugins/Runtime/Analytics/AnalyticsMulticast/Source/AnalyticsMulticast/Private/AnalyticsMulticast.cpp
new file mode 100644
--- /dev/null
+++ b/Engine/Plugins/Runtime/Analytics/AnalyticsMulticast/Source/AnalyticsMulticast/Private/AnalyticsMulticast.cpp
@@ -0,0 +1,319 @@
+// Copyright 1998-2015 Epic Games, Inc. All Rights Reserved.
+
+#include "AnalyticsMulticastPrivatePCH.h"
+
+#include "Core.h"
+#include "Guid.h"
+#include "Json.h"
+#include "SecureHash.h"
+#include "AnalyticsMulticast.h"
+#include "Analytics.h"
+
+DEFINE_LOG_CATEGORY_STATIC(LogAnalytics, Display, All);
+
+IMPLEMENT_MODULE( FAnalyticsMulticast, AnalyticsMulticast );
+
+class FAnalyticsProviderMulticast : public IAnalyticsProvider
+{
+public:
+	FAnalyticsProviderMulticast(const FAnalyticsMulticast::Config& ConfigValues, const FAnalytics::FProviderConfigurationDelegate& GetConfigValue);
+
+	virtual bool StartSession(const TArray<FAnalyticsEventAttribute>& Attributes) override;
+	virtual void EndSession() override;
+	virtual void FlushEvents() override;
+
+	virtual void SetUserID(const FString& InUserID) override;
+	virtual FString GetUserID() const override;
+
+	virtual FString GetSessionID() const override;
+	virtual bool SetSessionID(const FString& InSessionID) override;
+
+	virtual void RecordEvent(const FString& EventName, const TArray<FAnalyticsEventAttribute>& Attributes) override;
+	virtual void RecordItemPurchase(const FString& ItemId, const FString& Currency, int PerItemCost, int ItemQuantity) override;
+	virtual void RecordCurrencyPurchase(const FString& GameCurrencyType, int GameCurrencyAmount, const FString& RealCurrencyType, float RealMoneyCost, const FString& PaymentProvider) override;
+	virtual void RecordCurrencyGiven(const FString& GameCurrencyType, int GameCurrencyAmount) override;
+
+	virtual ~FAnalyticsProviderMulticast();
+
+	bool HasValidProviders() const { return Providers.Num() > 0; }
+
+	virtual void SetBuildInfo(const FString& InBuildInfo) override;
+	virtual void SetGender(const FString& InGender) override;
+	virtual void SetLocation(const FString& InLocation) override;
+	virtual void SetAge(const int32 InAge) override;
+
+	virtual void RecordItemPurchase(const FString& ItemId, int ItemQuantity, const TArray<FAnalyticsEventAttribute>& EventAttrs) override;
+	virtual void RecordCurrencyPurchase(const FString& GameCurrencyType, int GameCurrencyAmount, const TArray<FAnalyticsEventAttribute>& EventAttrs) override;
+	virtual void RecordCurrencyGiven(const FString& GameCurrencyType, int GameCurrencyAmount, const TArray<FAnalyticsEventAttribute>& EventAttrs) override;
+	virtual void RecordError(const FString& Error, const TArray<FAnalyticsEventAttribute>& EventAttrs) override;
+	virtual void RecordProgress(const FString& ProgressType, const FString& ProgressHierarchy, const TArray<FAnalyticsEventAttribute>& EventAttrs) override;
+	virtual void RecordProgress(const FString& ProgressType, const TArray<FString>& ProgressHierarchy, const TArray<FAnalyticsEventAttribute>& EventAttrs) override;
+
+private:
+	TArray<TSharedPtr<IAnalyticsProvider> > Providers;
+	TArray<FString> ProviderModules;
+};
+
+void FAnalyticsMulticast::StartupModule()
+{
+}
+
+void FAnalyticsMulticast::ShutdownModule()
+{
+}
+
+TSharedPtr<IAnalyticsProvider> FAnalyticsMulticast::CreateAnalyticsProvider(const FAnalytics::FProviderConfigurationDelegate& GetConfigValue) const
+{
+	if (GetConfigValue.IsBound())
+	{
+		Config ConfigValues;
+		ConfigValues.ProviderModuleNames = GetConfigValue.Execute(Config::GetKeyNameForProviderModuleNames(), true);
+		if (ConfigValues.ProviderModuleNames.IsEmpty())
+		{
+			UE_LOG(LogAnalytics, Warning, TEXT("CreateAnalyticsProvider delegate did not contain required parameter %s"), *Config::GetKeyNameForProviderModuleNames());
+			return NULL;
+		}
+		return CreateAnalyticsProvider(ConfigValues, GetConfigValue);
+	}
+	else
+	{
+		UE_LOG(LogAnalytics, Warning, TEXT("CreateAnalyticsProvider called with an unbound delegate"));
+	}
+	return NULL;
+}
+
+TSharedPtr<IAnalyticsProvider> FAnalyticsMulticast::CreateAnalyticsProvider(const Config& ConfigValues, const FAnalytics::FProviderConfigurationDelegate& GetConfigValue) const
+{
+	TScopedPointer<FAnalyticsProviderMulticast> Provider(new FAnalyticsProviderMulticast(ConfigValues, GetConfigValue));
+	// If we didn't have any valid providers, return NULL
+	if (!Provider->HasValidProviders())
+	{
+		return NULL;
+	}
+	// if we were configured ok, bind thee result to a shared pointer.
+	return TSharedPtr<IAnalyticsProvider>(Provider.Release());
+}
+
+/**
+ * Perform any initialization.
+ */
+FAnalyticsProviderMulticast::FAnalyticsProviderMulticast(const FAnalyticsMulticast::Config& ConfigValues, const FAnalytics::FProviderConfigurationDelegate& GetConfigValue)
+{
+	UE_LOG(LogAnalytics, Verbose, TEXT("Initializing Multicast Analytics provider"));
+
+	if (GetConfigValue.IsBound())
+	{
+		TArray<FString> ModuleNamesArray;
+		ConfigValues.ProviderModuleNames.ParseIntoArray(ModuleNamesArray, TEXT(","), true);
+		for (TArray<FString>::TConstIterator it(ModuleNamesArray);it;++it)
+		{
+			TSharedPtr<IAnalyticsProvider> NewProvider = FAnalytics::Get().CreateAnalyticsProvider(FName(**it), GetConfigValue);
+			if (NewProvider.IsValid())
+			{
+				Providers.Add(NewProvider);
+				ProviderModules.Add(*it);
+			}
+		}
+	}
+}
+
+FAnalyticsProviderMulticast::~FAnalyticsProviderMulticast()
+{
+	UE_LOG(LogAnalytics, Verbose, TEXT("Destroying Multicast Analytics provider"));
+}
+
+bool FAnalyticsProviderMulticast::StartSession(const TArray<FAnalyticsEventAttribute>& Attributes)
+{
+	bool bResult = false;
+	for (TArray<TSharedPtr<IAnalyticsProvider> >::TConstIterator it(Providers);it;++it)
+	{
+		bResult |= (*it)->StartSession(Attributes);
+	}
+	return bResult;
+}
+
+void FAnalyticsProviderMulticast::EndSession()
+{
+	for (TArray<TSharedPtr<IAnalyticsProvider> >::TConstIterator it(Providers);it;++it)
+	{
+		(*it)->EndSession();
+	}
+}
+
+void FAnalyticsProviderMulticast::FlushEvents()
+{
+	for (TArray<TSharedPtr<IAnalyticsProvider> >::TConstIterator it(Providers);it;++it)
+	{
+		(*it)->FlushEvents();
+	}
+}
+
+void FAnalyticsProviderMulticast::SetUserID( const FString& InUserID )
+{
+	for (TArray<TSharedPtr<IAnalyticsProvider> >::TConstIterator it(Providers);it;++it)
+	{
+		(*it)->SetUserID(InUserID);
+	}
+}
+
+FString FAnalyticsProviderMulticast::GetUserID() const 
+{
+	for (TArray<TSharedPtr<IAnalyticsProvider> >::TConstIterator it(Providers);it;++it)
+	{
+		return (*it)->GetUserID();
+	}
+	return FString();
+}
+
+FString FAnalyticsProviderMulticast::GetSessionID() const
+{
+	// combine all session IDs into Module@@Session##Module@@Session...
+	FString Result;
+	for (int i=0;i<Providers.Num();++i)
+	{
+		if (i>0) Result += "##";
+		Result += ProviderModules[i] + "@@" + Providers[i]->GetSessionID();
+	}
+	return Result;
+}
+
+bool FAnalyticsProviderMulticast::SetSessionID(const FString& InSessionID)
+{
+	bool bResult = false;
+	// parse out the format from GetSessionID and set the SessionID in each provider.
+	TArray<FString> SessionIDs;
+	// parse out the module/session pairs
+	InSessionID.ParseIntoArray(SessionIDs, TEXT("##"), true);
+	for (TArray<FString>::TConstIterator it(SessionIDs);it;++it)
+	{
+		FString ModuleName, SessionID;
+		// split out the module name/sessionID pair
+		if (it->Split(TEXT("@@"), &ModuleName, &SessionID))
+		{
+			// find the module name in the list of existing providers
+			for (int i=0;i<Providers.Num();++i)
+			{
+				// only set the session ID if it is not empty.
+				if (ProviderModules[i] == ModuleName && !SessionID.IsEmpty())
+				{
+					bResult |= Providers[i]->SetSessionID(SessionID);
+				}
+			}
+		}
+	}
+	return bResult;
+}
+
+
+void FAnalyticsProviderMulticast::RecordEvent( const FString& EventName, const TArray<FAnalyticsEventAttribute>& Attributes )
+{
+	for (TArray<TSharedPtr<IAnalyticsProvider> >::TConstIterator it(Providers);it;++it)
+	{
+		(*it)->RecordEvent(EventName, Attributes);
+	}
+}
+
+void FAnalyticsProviderMulticast::RecordItemPurchase( const FString& ItemId, const FString& Currency, int PerItemCost, int ItemQuantity )
+{
+	for (TArray<TSharedPtr<IAnalyticsProvider> >::TConstIterator it(Providers);it;++it)
+	{
+		(*it)->RecordItemPurchase(ItemId, Currency, PerItemCost, ItemQuantity);
+	}
+}
+
+void FAnalyticsProviderMulticast::RecordCurrencyPurchase( const FString& GameCurrencyType, int GameCurrencyAmount, const FString& RealCurrencyType, float RealMoneyCost, const FString& PaymentProvider )
+{
+	for (TArray<TSharedPtr<IAnalyticsProvider> >::TConstIterator it(Providers);it;++it)
+	{
+		(*it)->RecordCurrencyPurchase(GameCurrencyType, GameCurrencyAmount, RealCurrencyType, RealMoneyCost, PaymentProvider);
+	}
+}
+
+void FAnalyticsProviderMulticast::RecordCurrencyGiven( const FString& GameCurrencyType, int GameCurrencyAmount )
+{
+	for (TArray<TSharedPtr<IAnalyticsProvider> >::TConstIterator it(Providers);it;++it)
+	{
+		(*it)->RecordCurrencyGiven(GameCurrencyType, GameCurrencyAmount);
+	}
+}
+
+void FAnalyticsProviderMulticast::SetBuildInfo(const FString& InBuildInfo)
+{
+	for (auto Provider : Providers)
+	{
+		Provider->SetBuildInfo(InBuildInfo);
+	}
+}
+
+void FAnalyticsProviderMulticast::SetGender(const FString& InGender)
+{
+	for (auto Provider : Providers)
+	{
+		Provider->SetGender(InGender);
+	}
+}
+
+void FAnalyticsProviderMulticast::SetLocation(const FString& InLocation)
+{
+	for (auto Provider : Providers)
+	{
+		Provider->SetLocation(InLocation);
+	}
+}
+
+void FAnalyticsProviderMulticast::SetAge(const int32 InAge)
+{
+	for (auto Provider : Providers)
+	{
+		Provider->SetAge(InAge);
+	}
+}
+
+void FAnalyticsProviderMulticast::RecordItemPurchase(const FString& ItemId, int ItemQuantity, const TArray<FAnalyticsEventAttribute>& EventAttrs)
+{
+	for (auto Provider : Providers)
+	{
+		Provider->RecordItemPurchase(ItemId, ItemQuantity, EventAttrs);
+	}
+}
+
+void FAnalyticsProviderMulticast::RecordCurrencyPurchase(const FString& GameCurrencyType, int GameCurrencyAmount, const TArray<FAnalyticsEventAttribute>& EventAttrs)
+{
+	for (auto Provider : Providers)
+	{
+		Provider->RecordCurrencyPurchase(GameCurrencyType, GameCurrencyAmount, EventAttrs);
+	}
+}
+
+void FAnalyticsProviderMulticast::RecordCurrencyGiven(const FString& GameCurrencyType, int GameCurrencyAmount, const TArray<FAnalyticsEventAttribute>& EventAttrs)
+{
+	for (auto Provider : Providers)
+	{
+		Provider->RecordCurrencyGiven(GameCurrencyType, GameCurrencyAmount, EventAttrs);
+	}
+}
+
+void FAnalyticsProviderMulticast::RecordError(const FString& Error, const TArray<FAnalyticsEventAttribute>& EventAttrs)
+{
+	for (auto Provider : Providers)
+	{
+		Provider->RecordError(Error, EventAttrs);
+	}
+}
+
+void FAnalyticsProviderMulticast::RecordProgress(const FString& ProgressType, const FString& ProgressHierarchy, const TArray<FAnalyticsEventAttribute>& EventAttrs)
+{
+	for (auto Provider : Providers)
+	{
+		Provider->RecordProgress(ProgressType, ProgressHierarchy, EventAttrs);
+	}
+}
+
+void FAnalyticsProviderMulticast::RecordProgress(const FString& ProgressType, const TArray<FString>& ProgressHierarchy, const TArray<FAnalyticsEventAttribute>& EventAttrs)
+{
+	for (auto Provider : Providers)
+	{
+		Provider->RecordProgress(ProgressType, ProgressHierarchy, EventAttrs);
+	}
+}
+
diff --git a/Engine/Plugins/Runtime/Analytics/AnalyticsMulticast/Source/AnalyticsMulticast/Private/AnalyticsMulticastPrivatePCH.h b/Engine/Plugins/Runtime/Analytics/AnalyticsMulticast/Source/AnalyticsMulticast/Private/AnalyticsMulticastPrivatePCH.h
new file mode 100644
--- /dev/null
+++ b/Engine/Plugins/Runtime/Analytics/AnalyticsMulticast/Source/AnalyticsMulticast/Private/AnalyticsMulticastPrivatePCH.h
@@ -0,0 +1,4 @@
+// Copyright 1998-2015 Epic Games, Inc. All Rights Reserved.
+
+// You should place include statements to your module's private header files here.  You only need to
+// add includes for headers that are used in most of your module's source files though.
diff --git a/Engine/Plugins/Runtime/Analytics/AnalyticsMulticast/Source/AnalyticsMulticast/Public/AnalyticsMulticast.h b/Engine/Plugins/Runtime/Analytics/AnalyticsMulticast/Source/AnalyticsMulticast/Public/AnalyticsMulticast.h
new file mode 100644
--- /dev/null
+++ b/Engine/Plugins/Runtime/Analytics/AnalyticsMulticast/Source/AnalyticsMulticast/Public/AnalyticsMulticast.h
@@ -0,0 +1,76 @@
+// Copyright 1998-2015 Epic Games, Inc. All Rights Reserved.
+
+#pragma once
+
+#include "IAnalyticsProviderModule.h"
+#include "Core.h"
+
+class IAnalyticsProvider;
+
+/**
+ * Exposes a multicast provider that multicasts analytics events to multiple providers.
+ * Configured using a comma separated list of provider modules. Each module then uses
+ * The supplied configuration delegate to configure itself.
+ */
+class FAnalyticsMulticast : public IAnalyticsProviderModule
+{
+	//--------------------------------------------------------------------------
+	// Module functionality
+	//--------------------------------------------------------------------------
+public:
+	/**
+	 * Singleton-like access to this module's interface.  This is just for convenience!
+	 * Beware of calling this during the shutdown phase, though.  Your module might have been unloaded already.
+	 *
+	 * @return Returns singleton instance, loading the module on demand if needed
+	 */
+	static inline FAnalyticsMulticast& Get()
+	{
+		return FModuleManager::LoadModuleChecked< FAnalyticsMulticast >( "AnalyticsMulticast" );
+	}
+
+	//--------------------------------------------------------------------------
+	// Configuration functionality
+	//--------------------------------------------------------------------------
+public:
+	/** 
+	 * Defines required configuration values for multicast analytics provider. 
+	 * Basically, you provide a list of provider modules that define the providers you want
+	 * to multicast events to. Beyond that, each provider module created will use the 
+	 * provided configuration delegate to configure itself, so that configuration delegate
+	 * must be able to configure each specific provider as well (see CreateAnalyticsProvider function below).
+	 */
+	struct Config
+	{
+		/** Comma separated list of analytics provider modules */
+		FString ProviderModuleNames;
+
+		/** KeyName required for APIKey configuration. */
+		static FString GetKeyNameForProviderModuleNames() { return TEXT("ProviderModuleNames"); }
+	};
+
+	//--------------------------------------------------------------------------
+	// provider factory functions
+	//--------------------------------------------------------------------------
+public:
+	/**
+	 * IAnalyticsProviderModule interface.
+	 * Creates the analytics provider given a configuration delegate.
+	 * The keys required exactly match the field names in the Config object. 
+	 * 
+	 * When a particular provider module is loaded, it will create an instance and use the
+	 * provided Configuration delegate to configure each provider.
+	 */
+	virtual TSharedPtr<IAnalyticsProvider> CreateAnalyticsProvider(const FAnalytics::FProviderConfigurationDelegate& GetConfigValue) const;
+
+	
+	/** 
+	 * Construct an analytics provider directly from a config object (and a delegate to provide configuration to each configured provider).
+	 */
+	virtual TSharedPtr<IAnalyticsProvider> CreateAnalyticsProvider(const Config& ConfigValues, const FAnalytics::FProviderConfigurationDelegate& GetConfigValue) const;
+
+private:
+	virtual void StartupModule() override;
+	virtual void ShutdownModule() override;
+};
+
diff --git a/Engine/Source/Runtime/Analytics/AnalyticsMulticast/AnalyticsMulticast.Build.cs b/Engine/Source/Runtime/Analytics/AnalyticsMulticast/AnalyticsMulticast.Build.cs
deleted file mode 100644
diff --git a/Engine/Source/Runtime/Analytics/AnalyticsMulticast/Private/AnalyticsMulticast.cpp b/Engine/Source/Runtime/Analytics/AnalyticsMulticast/Private/AnalyticsMulticast.cpp
deleted file mode 100644
diff --git a/Engine/Source/Runtime/Analytics/AnalyticsMulticast/Private/AnalyticsMulticastPrivatePCH.h b/Engine/Source/Runtime/Analytics/AnalyticsMulticast/Private/AnalyticsMulticastPrivatePCH.h
deleted file mode 100644
diff --git a/Engine/Source/Runtime/Analytics/AnalyticsMulticast/Public/AnalyticsMulticast.h b/Engine/Source/Runtime/Analytics/AnalyticsMulticast/Public/AnalyticsMulticast.h
deleted file mode 100644
