Improved the singleton performance of the Apsalar and Multicast analytics plugins

diff --git a/Engine/Plugins/Runtime/Analytics/AnalyticsMulticast/Source/AnalyticsMulticast/Private/AnalyticsMulticast.cpp b/Engine/Plugins/Runtime/Analytics/AnalyticsMulticast/Source/AnalyticsMulticast/Private/AnalyticsMulticast.cpp
--- a/Engine/Plugins/Runtime/Analytics/AnalyticsMulticast/Source/AnalyticsMulticast/Private/AnalyticsMulticast.cpp
+++ b/Engine/Plugins/Runtime/Analytics/AnalyticsMulticast/Source/AnalyticsMulticast/Private/AnalyticsMulticast.cpp
@@ -15,9 +15,24 @@ IMPLEMENT_MODULE( FAnalyticsMulticast, AnalyticsMulticast );
 
 class FAnalyticsProviderMulticast : public IAnalyticsProvider
 {
-public:
+	/** Singleton for analytics */
+	static TSharedPtr<IAnalyticsProvider> Provider;
 	FAnalyticsProviderMulticast(const FAnalyticsMulticast::Config& ConfigValues, const FAnalytics::FProviderConfigurationDelegate& GetConfigValue);
 
+public:
+	static TSharedPtr<IAnalyticsProvider> Create(const FAnalyticsMulticast::Config& ConfigValues, const FAnalytics::FProviderConfigurationDelegate& GetConfigValue)
+	{
+		if (!Provider.IsValid())
+		{
+			Provider = TSharedPtr<IAnalyticsProvider>(new FAnalyticsProviderMulticast(ConfigValues, GetConfigValue));
+		}
+		return Provider;
+	}
+	static void Destroy()
+	{
+		Provider.Reset();
+	}
+
 	virtual bool StartSession(const TArray<FAnalyticsEventAttribute>& Attributes) override;
 	virtual void EndSession() override;
 	virtual void FlushEvents() override;
@@ -54,12 +69,15 @@ private:
 	TArray<FString> ProviderModules;
 };
 
+TSharedPtr<IAnalyticsProvider> FAnalyticsProviderMulticast::Provider;
+
 void FAnalyticsMulticast::StartupModule()
 {
 }
 
 void FAnalyticsMulticast::ShutdownModule()
 {
+	FAnalyticsProviderMulticast::Destroy();
 }
 
 TSharedPtr<IAnalyticsProvider> FAnalyticsMulticast::CreateAnalyticsProvider(const FAnalytics::FProviderConfigurationDelegate& GetConfigValue) const
@@ -84,14 +102,7 @@ TSharedPtr<IAnalyticsProvider> FAnalyticsMulticast::CreateAnalyticsProvider(cons
 
 TSharedPtr<IAnalyticsProvider> FAnalyticsMulticast::CreateAnalyticsProvider(const Config& ConfigValues, const FAnalytics::FProviderConfigurationDelegate& GetConfigValue) const
 {
-	TScopedPointer<FAnalyticsProviderMulticast> Provider(new FAnalyticsProviderMulticast(ConfigValues, GetConfigValue));
-	// If we didn't have any valid providers, return NULL
-	if (!Provider->HasValidProviders())
-	{
-		return NULL;
-	}
-	// if we were configured ok, bind thee result to a shared pointer.
-	return TSharedPtr<IAnalyticsProvider>(Provider.Release());
+	return FAnalyticsProviderMulticast::Create(ConfigValues, GetConfigValue);
 }
 
 /**
diff --git a/Engine/Plugins/Runtime/Analytics/Apsalar/IOSApsalar/Source/IOSApsalar/Private/IOSApsalar.cpp b/Engine/Plugins/Runtime/Analytics/Apsalar/IOSApsalar/Source/IOSApsalar/Private/IOSApsalar.cpp
--- a/Engine/Plugins/Runtime/Analytics/Apsalar/IOSApsalar/Source/IOSApsalar/Private/IOSApsalar.cpp
+++ b/Engine/Plugins/Runtime/Analytics/Apsalar/IOSApsalar/Source/IOSApsalar/Private/IOSApsalar.cpp
@@ -6,12 +6,15 @@ DEFINE_LOG_CATEGORY_STATIC(LogAnalytics, Display, All);
 
 IMPLEMENT_MODULE( FAnalyticsIOSApsalar, IOSApsalar )
 
+TSharedPtr<IAnalyticsProvider> FAnalyticsProviderApsalar::Provider;
+
 void FAnalyticsIOSApsalar::StartupModule()
 {
 }
 
 void FAnalyticsIOSApsalar::ShutdownModule()
 {
+	FAnalyticsProviderApsalar::Destroy();
 }
 
 TSharedPtr<IAnalyticsProvider> FAnalyticsIOSApsalar::CreateAnalyticsProvider(const FAnalytics::FProviderConfigurationDelegate& GetConfigValue) const
@@ -23,8 +26,8 @@ TSharedPtr<IAnalyticsProvider> FAnalyticsIOSApsalar::CreateAnalyticsProvider(con
 		const FString SendInterval = GetConfigValue.Execute(TEXT("SendInterval"), false);
 		const FString MaxBufferSize = GetConfigValue.Execute(TEXT("MaxBufferSize"), false);
 		const FString ManuallyReportRevenue = GetConfigValue.Execute(TEXT("ManuallyReportRevenue"), false);
-		const bool bWantsManualRevenueReporting = ManuallyReportRevenue.Compare(TEXT("true"), ESearchCase::IgnoreCase);
-		return TSharedPtr<IAnalyticsProvider>(new FAnalyticsProviderApsalar(Key, Secret, FCString::Atoi(*SendInterval), FCString::Atoi(*MaxBufferSize), bWantsManualRevenueReporting));
+		const bool bWantsManualRevenueReporting = ManuallyReportRevenue.Compare(TEXT("true"), ESearchCase::IgnoreCase) == 0;
+		return FAnalyticsProviderApsalar::Create(Key, Secret, FCString::Atoi(*SendInterval), FCString::Atoi(*MaxBufferSize), bWantsManualRevenueReporting);
 	}
 	else
 	{
diff --git a/Engine/Plugins/Runtime/Analytics/Apsalar/IOSApsalar/Source/IOSApsalar/Private/IOSApsalarProvider.h b/Engine/Plugins/Runtime/Analytics/Apsalar/IOSApsalar/Source/IOSApsalar/Private/IOSApsalarProvider.h
--- a/Engine/Plugins/Runtime/Analytics/Apsalar/IOSApsalar/Source/IOSApsalar/Private/IOSApsalarProvider.h
+++ b/Engine/Plugins/Runtime/Analytics/Apsalar/IOSApsalar/Source/IOSApsalar/Private/IOSApsalarProvider.h
@@ -12,8 +12,24 @@ class FAnalyticsProviderApsalar :
 	/** Tracks whether we need to start the session or restart it */
 	bool bHasSessionStarted;
 
-public:
+	/** Singleton for analytics */
+	static TSharedPtr<IAnalyticsProvider> Provider;
 	FAnalyticsProviderApsalar(const FString Key, const FString Secret, const int32 SendInterval, const int32 MaxBufferSize, const bool bWantsManualRevenueReporting);
+
+public:
+	static TSharedPtr<IAnalyticsProvider> Create(const FString Key, const FString Secret, const int32 SendInterval, const int32 MaxBufferSize, const bool bWantsManualRevenueReporting)
+	{
+		if (!Provider.IsValid())
+		{
+			Provider = TSharedPtr<IAnalyticsProvider>(new FAnalyticsProviderApsalar(Key, Secret, SendInterval, MaxBufferSize, bWantsManualRevenueReporting));
+		}
+		return Provider;
+	}
+	static void Destroy()
+	{
+		Provider.Reset();
+	}
+
 	virtual ~FAnalyticsProviderApsalar();
 
 	virtual bool StartSession(const TArray<FAnalyticsEventAttribute>& Attributes) override;
