Fix build error variable shadowing.

diff --git a/Engine/Plugins/Runtime/Analytics/AnalyticsMulticast/Source/AnalyticsMulticast/Private/AnalyticsMulticast.cpp b/Engine/Plugins/Runtime/Analytics/AnalyticsMulticast/Source/AnalyticsMulticast/Private/AnalyticsMulticast.cpp
--- a/Engine/Plugins/Runtime/Analytics/AnalyticsMulticast/Source/AnalyticsMulticast/Private/AnalyticsMulticast.cpp
+++ b/Engine/Plugins/Runtime/Analytics/AnalyticsMulticast/Source/AnalyticsMulticast/Private/AnalyticsMulticast.cpp
@@ -250,81 +250,81 @@ void FAnalyticsProviderMulticast::RecordCurrencyGiven( const FString& GameCurren
 
 void FAnalyticsProviderMulticast::SetBuildInfo(const FString& InBuildInfo)
 {
-	for (auto Provider : Providers)
+	for (auto CurrentProvider : Providers)
 	{
-		Provider->SetBuildInfo(InBuildInfo);
+		CurrentProvider->SetBuildInfo(InBuildInfo);
 	}
 }
 
 void FAnalyticsProviderMulticast::SetGender(const FString& InGender)
 {
-	for (auto Provider : Providers)
+	for (auto CurrentProvider : Providers)
 	{
-		Provider->SetGender(InGender);
+		CurrentProvider->SetGender(InGender);
 	}
 }
 
 void FAnalyticsProviderMulticast::SetLocation(const FString& InLocation)
 {
-	for (auto Provider : Providers)
+	for (auto CurrentProvider : Providers)
 	{
-		Provider->SetLocation(InLocation);
+		CurrentProvider->SetLocation(InLocation);
 	}
 }
 
 void FAnalyticsProviderMulticast::SetAge(const int32 InAge)
 {
-	for (auto Provider : Providers)
+	for (auto CurrentProvider : Providers)
 	{
-		Provider->SetAge(InAge);
+		CurrentProvider->SetAge(InAge);
 	}
 }
 
 void FAnalyticsProviderMulticast::RecordItemPurchase(const FString& ItemId, int ItemQuantity, const TArray<FAnalyticsEventAttribute>& EventAttrs)
 {
-	for (auto Provider : Providers)
+	for (auto CurrentProvider : Providers)
 	{
-		Provider->RecordItemPurchase(ItemId, ItemQuantity, EventAttrs);
+		CurrentProvider->RecordItemPurchase(ItemId, ItemQuantity, EventAttrs);
 	}
 }
 
 void FAnalyticsProviderMulticast::RecordCurrencyPurchase(const FString& GameCurrencyType, int GameCurrencyAmount, const TArray<FAnalyticsEventAttribute>& EventAttrs)
 {
-	for (auto Provider : Providers)
+	for (auto CurrentProvider : Providers)
 	{
-		Provider->RecordCurrencyPurchase(GameCurrencyType, GameCurrencyAmount, EventAttrs);
+		CurrentProvider->RecordCurrencyPurchase(GameCurrencyType, GameCurrencyAmount, EventAttrs);
 	}
 }
 
 void FAnalyticsProviderMulticast::RecordCurrencyGiven(const FString& GameCurrencyType, int GameCurrencyAmount, const TArray<FAnalyticsEventAttribute>& EventAttrs)
 {
-	for (auto Provider : Providers)
+	for (auto CurrentProvider : Providers)
 	{
-		Provider->RecordCurrencyGiven(GameCurrencyType, GameCurrencyAmount, EventAttrs);
+		CurrentProvider->RecordCurrencyGiven(GameCurrencyType, GameCurrencyAmount, EventAttrs);
 	}
 }
 
 void FAnalyticsProviderMulticast::RecordError(const FString& Error, const TArray<FAnalyticsEventAttribute>& EventAttrs)
 {
-	for (auto Provider : Providers)
+	for (auto CurrentProvider : Providers)
 	{
-		Provider->RecordError(Error, EventAttrs);
+		CurrentProvider->RecordError(Error, EventAttrs);
 	}
 }
 
 void FAnalyticsProviderMulticast::RecordProgress(const FString& ProgressType, const FString& ProgressHierarchy, const TArray<FAnalyticsEventAttribute>& EventAttrs)
 {
-	for (auto Provider : Providers)
+	for (auto CurrentProvider : Providers)
 	{
-		Provider->RecordProgress(ProgressType, ProgressHierarchy, EventAttrs);
+		CurrentProvider->RecordProgress(ProgressType, ProgressHierarchy, EventAttrs);
 	}
 }
 
 void FAnalyticsProviderMulticast::RecordProgress(const FString& ProgressType, const TArray<FString>& ProgressHierarchy, const TArray<FAnalyticsEventAttribute>& EventAttrs)
 {
-	for (auto Provider : Providers)
+	for (auto CurrentProvider : Providers)
 	{
-		Provider->RecordProgress(ProgressType, ProgressHierarchy, EventAttrs);
+		CurrentProvider->RecordProgress(ProgressType, ProgressHierarchy, EventAttrs);
 	}
 }
 
