Optimized include module name dependencies.

diff --git a/Engine/Plugins/Media/AjaMedia/Source/AjaMediaEditor/AjaMediaEditor.Build.cs b/Engine/Plugins/Media/AjaMedia/Source/AjaMediaEditor/AjaMediaEditor.Build.cs
--- a/Engine/Plugins/Media/AjaMedia/Source/AjaMediaEditor/AjaMediaEditor.Build.cs
+++ b/Engine/Plugins/Media/AjaMedia/Source/AjaMediaEditor/AjaMediaEditor.Build.cs
@@ -28,7 +28,6 @@ namespace UnrealBuildTool.Rules
 
 			PrivateIncludePathModuleNames.AddRange(
 				new string[] {
-					"AjaCore",
 					"AssetTools"
 				});
 		}
diff --git a/Engine/Plugins/Runtime/Analytics/Adjust/Source/AdjustEditor/AdjustEditor.Build.cs b/Engine/Plugins/Runtime/Analytics/Adjust/Source/AdjustEditor/AdjustEditor.Build.cs
--- a/Engine/Plugins/Runtime/Analytics/Adjust/Source/AdjustEditor/AdjustEditor.Build.cs
+++ b/Engine/Plugins/Runtime/Analytics/Adjust/Source/AdjustEditor/AdjustEditor.Build.cs
@@ -17,12 +17,5 @@ public class AdjustEditor : ModuleRules
 				"Projects"
 			}
 			);
-
-		PrivateIncludePathModuleNames.AddRange(
-			new string[]
-            {
-				"Settings"
-			}
-		);
 	}
 }
diff --git a/Engine/Plugins/Runtime/Analytics/AnalyticsMulticast/Source/AnaltyicsMulticastEditor/AnalyticsMulticastEditor.Build.cs b/Engine/Plugins/Runtime/Analytics/AnalyticsMulticast/Source/AnaltyicsMulticastEditor/AnalyticsMulticastEditor.Build.cs
--- a/Engine/Plugins/Runtime/Analytics/AnalyticsMulticast/Source/AnaltyicsMulticastEditor/AnalyticsMulticastEditor.Build.cs
+++ b/Engine/Plugins/Runtime/Analytics/AnalyticsMulticast/Source/AnaltyicsMulticastEditor/AnalyticsMulticastEditor.Build.cs
@@ -25,13 +25,6 @@ public class AnalyticsMulticastEditor : ModuleRules
 			}
 			);
 
-		PrivateIncludePathModuleNames.AddRange(
-			new string[] {
-				"Settings",
-                "AssetTools"
-			}
-		);
-
         DynamicallyLoadedModuleNames.AddRange(
             new string[] {
                 "AssetTools"
diff --git a/Engine/Plugins/Runtime/Analytics/Flurry/Source/FlurryEditor/FlurryEditor.Build.cs b/Engine/Plugins/Runtime/Analytics/Flurry/Source/FlurryEditor/FlurryEditor.Build.cs
--- a/Engine/Plugins/Runtime/Analytics/Flurry/Source/FlurryEditor/FlurryEditor.Build.cs
+++ b/Engine/Plugins/Runtime/Analytics/Flurry/Source/FlurryEditor/FlurryEditor.Build.cs
@@ -17,12 +17,5 @@ public class FlurryEditor : ModuleRules
 				"DeveloperSettings"
 			}
 			);
-
-		PrivateIncludePathModuleNames.AddRange(
-			new string[]
-            {
-				"Settings"
-			}
-		);
 	}
 }
diff --git a/Engine/Plugins/Runtime/Oculus/OculusAudio/Source/OculusAudio/OculusAudio.Build.cs b/Engine/Plugins/Runtime/Oculus/OculusAudio/Source/OculusAudio/OculusAudio.Build.cs
--- a/Engine/Plugins/Runtime/Oculus/OculusAudio/Source/OculusAudio/OculusAudio.Build.cs
+++ b/Engine/Plugins/Runtime/Oculus/OculusAudio/Source/OculusAudio/OculusAudio.Build.cs
@@ -8,13 +8,6 @@ namespace UnrealBuildTool.Rules
 	{
 		public OculusAudio(ReadOnlyTargetRules Target) : base(Target)
 		{
-			PrivateIncludePathModuleNames.AddRange(
-				new string[]
-				{
-					"TargetPlatform"
-				}
-				);
-
 			PrivateIncludePaths.AddRange(
 				new string[] {
 					System.IO.Path.Combine(GetModuleDirectory("AudioMixer"), "Private"),
diff --git a/Engine/Plugins/Runtime/SoundMod/Source/SoundModImporter/SoundModImporter.Build.cs b/Engine/Plugins/Runtime/SoundMod/Source/SoundModImporter/SoundModImporter.Build.cs
--- a/Engine/Plugins/Runtime/SoundMod/Source/SoundModImporter/SoundModImporter.Build.cs
+++ b/Engine/Plugins/Runtime/SoundMod/Source/SoundModImporter/SoundModImporter.Build.cs
@@ -22,7 +22,6 @@ public class SoundModImporter : ModuleRules
 		PrivateIncludePathModuleNames.AddRange(
 			new string[] {
 				"AssetTools",
-				"AssetRegistry"
 			});
 
 		DynamicallyLoadedModuleNames.AddRange(
diff --git a/Engine/Source/Developer/TextureCompressor/TextureCompressor.Build.cs b/Engine/Source/Developer/TextureCompressor/TextureCompressor.Build.cs
--- a/Engine/Source/Developer/TextureCompressor/TextureCompressor.Build.cs
+++ b/Engine/Source/Developer/TextureCompressor/TextureCompressor.Build.cs
@@ -8,7 +8,6 @@ public class TextureCompressor : ModuleRules
 	{
 		PrivateIncludePathModuleNames.AddRange(new string[]
 		{
-			"CoreUObject",
 			"Engine",
 		});
 
